Rename WebTestsServer to WebTestsServerPy.

#Synced-CL 52039227

diff --git a/Engine/Source/Programs/AutomationTool/LowLevelTests/Tests/LowLevelTests.WebTestsLowLevelTestsExtension.cs b/Engine/Source/Programs/AutomationTool/LowLevelTests/Tests/LowLevelTests.WebTestsLowLevelTestsExtension.cs
--- a/Engine/Source/Programs/AutomationTool/LowLevelTests/Tests/LowLevelTests.WebTestsLowLevelTestsExtension.cs
+++ b/Engine/Source/Programs/AutomationTool/LowLevelTests/Tests/LowLevelTests.WebTestsLowLevelTestsExtension.cs
@@ -35,7 +35,7 @@ namespace LowLevelTests
 
 		private string WebTestsServerDir()
 		{
-			return Path.Combine(Unreal.EngineDirectory.FullName, "Source", "Programs", "WebTestsServer");
+			return Path.Combine(Unreal.EngineDirectory.FullName, "Source", "Programs", "WebTestsServerPy");
 		}
 
 		private void InstallWebServer()
@@ -64,7 +64,7 @@ namespace LowLevelTests
 			StartInfo.WorkingDirectory = WorkingDir;
 			StartInfo.FileName = PythonFile;
 			StartInfo.WindowStyle = ProcessWindowStyle.Normal;
-			StartInfo.Arguments = "manage.py runserver 0.0.0.0:8000";
+			StartInfo.Arguments = "-m daphne -e tcp:port=8000:interface=0.0.0.0 -e tcp:port=8001:interface=0.0.0.0 webtests.asgi:application";
 			StartInfo.UseShellExecute = true;
 			StartInfo.CreateNoWindow = false;
 
@@ -80,7 +80,7 @@ namespace LowLevelTests
 			Stopwatch sw = new Stopwatch();
 			sw.Start();
 
-			while (!IsServerPortOpen(UnrealHelpers.GetHostIpAddress(), 8000))
+			while (!IsServerPortOpen(UnrealHelpers.GetHostIpAddress(), 8000) || !IsServerPortOpen(UnrealHelpers.GetHostIpAddress(), 8001))
 			{
 				if (sw.ElapsedMilliseconds > 60000)
 				{
diff --git a/Engine/Source/Programs/WebTestsServer/createenv.bat b/Engine/Source/Programs/WebTestsServer/createenv.bat
deleted file mode 100755
diff --git a/Engine/Source/Programs/WebTestsServer/createenv.sh b/Engine/Source/Programs/WebTestsServer/createenv.sh
deleted file mode 100755
diff --git a/Engine/Source/Programs/WebTestsServer/dockerbuildandrun.bat b/Engine/Source/Programs/WebTestsServer/dockerbuildandrun.bat
deleted file mode 100755
diff --git a/Engine/Source/Programs/WebTestsServer/dockerfile b/Engine/Source/Programs/WebTestsServer/dockerfile
deleted file mode 100644
diff --git a/Engine/Source/Programs/WebTestsServer/httptests/__init__.py b/Engine/Source/Programs/WebTestsServer/httptests/__init__.py
deleted file mode 100644
diff --git a/Engine/Source/Programs/WebTestsServer/httptests/urls.py b/Engine/Source/Programs/WebTestsServer/httptests/urls.py
deleted file mode 100644
diff --git a/Engine/Source/Programs/WebTestsServer/httptests/views.py b/Engine/Source/Programs/WebTestsServer/httptests/views.py
deleted file mode 100644
diff --git a/Engine/Source/Programs/WebTestsServer/manage.py b/Engine/Source/Programs/WebTestsServer/manage.py
deleted file mode 100644
diff --git a/Engine/Source/Programs/WebTestsServer/requirements.txt b/Engine/Source/Programs/WebTestsServer/requirements.txt
deleted file mode 100644
diff --git a/Engine/Source/Programs/WebTestsServer/runserver.bat b/Engine/Source/Programs/WebTestsServer/runserver.bat
deleted file mode 100755
diff --git a/Engine/Source/Programs/WebTestsServer/runserver.sh b/Engine/Source/Programs/WebTestsServer/runserver.sh
deleted file mode 100644
diff --git a/Engine/Source/Programs/WebTestsServer/runserver_unixsocket.bat b/Engine/Source/Programs/WebTestsServer/runserver_unixsocket.bat
deleted file mode 100755
diff --git a/Engine/Source/Programs/WebTestsServer/unixsockets/Program.cs b/Engine/Source/Programs/WebTestsServer/unixsockets/Program.cs
deleted file mode 100644
diff --git a/Engine/Source/Programs/WebTestsServer/unixsockets/unixsocketsserver.csproj b/Engine/Source/Programs/WebTestsServer/unixsockets/unixsocketsserver.csproj
deleted file mode 100644
diff --git a/Engine/Source/Programs/WebTestsServer/websocketstests/__init__.py b/Engine/Source/Programs/WebTestsServer/websocketstests/__init__.py
deleted file mode 100644
diff --git a/Engine/Source/Programs/WebTestsServer/websocketstests/consumers.py b/Engine/Source/Programs/WebTestsServer/websocketstests/consumers.py
deleted file mode 100644
diff --git a/Engine/Source/Programs/WebTestsServer/websocketstests/routing.py b/Engine/Source/Programs/WebTestsServer/websocketstests/routing.py
deleted file mode 100644
diff --git a/Engine/Source/Programs/WebTestsServer/websocketstests/templates/websocketstests/index.html b/Engine/Source/Programs/WebTestsServer/websocketstests/templates/websocketstests/index.html
deleted file mode 100644
diff --git a/Engine/Source/Programs/WebTestsServer/websocketstests/urls.py b/Engine/Source/Programs/WebTestsServer/websocketstests/urls.py
deleted file mode 100644
diff --git a/Engine/Source/Programs/WebTestsServer/websocketstests/views.py b/Engine/Source/Programs/WebTestsServer/websocketstests/views.py
deleted file mode 100644
diff --git a/Engine/Source/Programs/WebTestsServer/webtests/__init__.py b/Engine/Source/Programs/WebTestsServer/webtests/__init__.py
deleted file mode 100644
diff --git a/Engine/Source/Programs/WebTestsServer/webtests/asgi.py b/Engine/Source/Programs/WebTestsServer/webtests/asgi.py
deleted file mode 100644
diff --git a/Engine/Source/Programs/WebTestsServer/webtests/settings.py b/Engine/Source/Programs/WebTestsServer/webtests/settings.py
deleted file mode 100644
diff --git a/Engine/Source/Programs/WebTestsServer/webtests/urls.py b/Engine/Source/Programs/WebTestsServer/webtests/urls.py
deleted file mode 100644
diff --git a/Engine/Source/Programs/WebTestsServer/webtests/wsgi.py b/Engine/Source/Programs/WebTestsServer/webtests/wsgi.py
deleted file mode 100644
diff --git a/Engine/Source/Programs/WebTestsServerPy/createenv.bat b/Engine/Source/Programs/WebTestsServerPy/createenv.bat
new file mode 100755
--- /dev/null
+++ b/Engine/Source/Programs/WebTestsServerPy/createenv.bat
@@ -0,0 +1,14 @@
+@echo off
+
+if not exist env (
+	..\..\..\Binaries\ThirdParty\Python3\Win64\python.exe -m venv env
+)
+
+call .\env\Scripts\activate.bat
+
+python -m pip install --upgrade pip
+pip install -r requirements.txt
+if %ERRORLEVEL% NEQ 0 (
+	rmdir /s /q env
+	exit %ERRORLEVEL%
+)
diff --git a/Engine/Source/Programs/WebTestsServerPy/createenv.sh b/Engine/Source/Programs/WebTestsServerPy/createenv.sh
new file mode 100755
--- /dev/null
+++ b/Engine/Source/Programs/WebTestsServerPy/createenv.sh
@@ -0,0 +1,22 @@
+if [ ! -d "env" ]; then
+	case "$OSTYPE" in
+		darwin*)
+			../../../Binaries/ThirdParty/Python3/Mac/bin/python3 -m venv env
+			;;
+		*)
+			# For daphne4.0.0 -> ImportError: libffi.so.6: cannot open shared object file: No such file or directory
+			curl -LO http://archive.ubuntu.com/ubuntu/pool/main/libf/libffi/libffi6_3.2.1-8_amd64.deb
+			sudo dpkg -i libffi6_3.2.1-8_amd64.deb
+			../../../Binaries/ThirdParty/Python3/Linux/bin/python3 -m venv env
+			;;
+	esac
+fi
+
+. ./env/bin/activate
+python -m pip install --upgrade pip
+pip install -r requirements.txt
+local install_exit_code = $?
+if [ $install_exit_code -ne 0 ]; then
+	rm -rf env
+	exit $install_exit_code
+fi
diff --git a/Engine/Source/Programs/WebTestsServerPy/dockerbuildandrun.bat b/Engine/Source/Programs/WebTestsServerPy/dockerbuildandrun.bat
new file mode 100755
--- /dev/null
+++ b/Engine/Source/Programs/WebTestsServerPy/dockerbuildandrun.bat
@@ -0,0 +1,14 @@
+@echo off
+
+:: Build image
+docker build . -t webtests
+
+set CONTAINER="webtestscontainer"
+:: Check if container exist and delete
+docker ps -a | findstr /c:"%CONTAINER%" >nul 2>&1
+if %errorlevel% equ 0 (
+	docker rm -f %CONTAINER%
+)
+
+:: Run container with image
+docker run --name %CONTAINER% -d -p 8000:8000 webtests
diff --git a/Engine/Source/Programs/WebTestsServerPy/dockerfile b/Engine/Source/Programs/WebTestsServerPy/dockerfile
new file mode 100644
--- /dev/null
+++ b/Engine/Source/Programs/WebTestsServerPy/dockerfile
@@ -0,0 +1,11 @@
+FROM python:3.10.4
+ENV DockerHOME=[REDACTED_USER_HOME]/webtests
+RUN mkdir -p $DockerHOME
+WORKDIR $DockerHOME
+ENV PYTHONDONTWRITEBYTECODE 1
+ENV PYTHONUNBUFFERED 1
+RUN pip install --upgrade pip
+COPY . $DockerHOME
+RUN pip install -r requirements.txt
+EXPOSE 8000
+CMD python manage.py runserver 0.0.0.0:8000
diff --git a/Engine/Source/Programs/WebTestsServerPy/httptests/urls.py b/Engine/Source/Programs/WebTestsServerPy/httptests/urls.py
new file mode 100644
--- /dev/null
+++ b/Engine/Source/Programs/WebTestsServerPy/httptests/urls.py
@@ -0,0 +1,18 @@
+from django.urls import path
+
+from . import views
+
+urlpatterns = [
+    path("httptests/methods", views.methods, name="methods"),
+    path("httptests/query_with_params/", views.query_with_params, name="query_with_params"),
+    path("httptests/get_data_without_chunks/<int:bytes_number>/<int:repeat_at>/", views.get_data_without_chunks, name="get_data_without_chunks"),
+    path("httptests/get_data_gzip/<int:bytes_number>/<int:repeat_at>/", views.get_data_gzip, name="get_data_gzip"),
+    path("httptests/streaming_download/<int:chunks>/<int:chunk_size>/<int:chunk_latency>/", views.streaming_download, name="streaming_download"),
+    path("httptests/streaming_upload_post", views.streaming_upload_post, name="streaming_upload_post"),
+    path("httptests/streaming_upload_put", views.streaming_upload_put, name="streaming_upload_put"),
+    path("httptests/redirect_from", views.redirect_from, name="redirect_from"),
+    path("httptests/redirect_to", views.redirect_to, name="redirect_to"),
+    path("httptests/mock_latency/<int:latency>/", views.mock_latency, name="mock_latency"),
+    path("httptests/mock_status/<int:status_code>/", views.mock_status, name="mock_status"),
+    path("httptests/post_echo_body", views.post_echo_body, name="post_echo_body"),
+]
diff --git a/Engine/Source/Programs/WebTestsServerPy/httptests/views.py b/Engine/Source/Programs/WebTestsServerPy/httptests/views.py
new file mode 100644
--- /dev/null
+++ b/Engine/Source/Programs/WebTestsServerPy/httptests/views.py
@@ -0,0 +1,114 @@
+import asyncio
+from datetime import datetime
+from django.shortcuts import render, redirect
+from django.http.response import JsonResponse
+from django.http import JsonResponse
+from django.http import HttpResponse
+from django.http import StreamingHttpResponse
+from django.views.decorators.gzip import gzip_page
+from rest_framework.decorators import api_view, parser_classes
+from rest_framework.parsers import FileUploadParser
+import time
+import json
+import logging
+
+logger = logging.getLogger('django')
+
+@api_view(['GET', 'POST', 'DELETE', 'PUT', 'PATCH', 'HEAD', 'OPTIONS', 'CONNECT', 'TRACE'])
+def methods(request):
+    return JsonResponse({})
+
+@api_view(['GET'])
+def query_with_params(request):
+    var_int = int(request.GET.get('var_int'))
+    var_str = request.GET.get('var_str')
+    return JsonResponse({'var_int' : var_int, 'var_str' : var_str}, content_type="application/json")
+
+def generate_data(bytes_number, repeat_at):
+    assert bytes_number > 0
+    assert repeat_at > 0
+    assert repeat_at <= 10
+    return [i%repeat_at for i in range(bytes_number)]
+
+@api_view(['GET'])
+def get_data_without_chunks(request, bytes_number, repeat_at):
+    return HttpResponse(generate_data(bytes_number, repeat_at))
+
+@api_view(['GET'])
+@gzip_page
+def get_data_gzip(request, bytes_number, repeat_at):
+    assert request.headers.get('Accept-Encoding') == 'gzip'
+    return HttpResponse(generate_data(bytes_number, repeat_at))
+
+@api_view(['GET'])
+def streaming_download(request, chunks, chunk_size, chunk_latency):
+    async def data_chunk_generator(): # using 'async' for StreamingHttpResponse in asgi
+        for x in range(chunks):
+            if chunk_latency > 0:
+                logger.debug("[%s]sleeping %d seconds", datetime.now().strftime("%H:%M:%S:%f"), chunk_latency)
+                await asyncio.sleep(chunk_latency)  # using 'await asyncio.sleep' instead of 'time.sleep' for StreamingHttpResponse in asgi
+            data_chunk = 'd' * chunk_size
+            logger.debug("[%s]sending %d bytes of data", datetime.now().strftime("%H:%M:%S:%f"), chunk_size)
+            yield data_chunk
+    response = StreamingHttpResponse(data_chunk_generator())
+    response['Content-Length'] = chunks * chunk_size
+    return response
+
+def output_file_info(f):
+    if f.multiple_chunks():
+        # big file
+        total_chunks = 0
+        total_bytes = 0
+        for chunk in f.chunks():
+            total_chunks += 1
+            total_bytes += len(chunk)
+        logger.info("Received %d chunks: %d bytes in total", total_chunks, total_bytes)
+    else:
+        # small file
+        logger.info("Received the whole file: %d bytes", f.size)
+
+@api_view(['POST'])
+def streaming_upload_post(request):
+    f = request.FILES['file']
+    output_file_info(f)
+    return JsonResponse({})
+
+@api_view(['PUT'])
+@parser_classes((FileUploadParser,))
+def streaming_upload_put(request):
+    f = request.data['file']
+    output_file_info(f)
+    return JsonResponse({})
+
+@api_view(['GET'])
+def redirect_from(request):
+    return redirect('redirect_to')
+
+@api_view(['GET'])
+def redirect_to(request):
+    return JsonResponse({})
+
+@api_view(['GET'])
+def mock_latency(request, latency):
+    time.sleep(latency)
+    return JsonResponse({})
+
+@api_view(['POST'])
+def post_echo_body(request):
+    return JsonResponse(request.data)
+
+@api_view(['GET', 'POST', 'DELETE', 'PUT'])
+def mock_status(request, status_code):
+    json_result = json.dumps({
+        "key_a": "value_a",
+        "key_b": "value_b"
+    })
+
+    headers_to_forward = ['Retry-After',]
+
+    response_headers = {}
+    for header_to_forward in headers_to_forward:
+        if header_to_forward in request.headers:
+            response_headers[header_to_forward] = request.headers[header_to_forward]
+
+    return HttpResponse(json_result, content_type="application/json", status=status_code, headers=response_headers)
diff --git a/Engine/Source/Programs/WebTestsServerPy/manage.py b/Engine/Source/Programs/WebTestsServerPy/manage.py
new file mode 100644
--- /dev/null
+++ b/Engine/Source/Programs/WebTestsServerPy/manage.py
@@ -0,0 +1,22 @@
+#!/usr/bin/env python
+"""Django's command-line utility for administrative tasks."""
+import os
+import sys
+
+
+def main():
+    """Run administrative tasks."""
+    os.environ.setdefault('DJANGO_SETTINGS_MODULE', 'webtests.settings')
+    try:
+        from django.core.management import execute_from_command_line
+    except ImportError as exc:
+        raise ImportError(
+            "Couldn't import Django. Are you sure it's installed and "
+            "available on your PYTHONPATH environment variable? Did you "
+            "forget to activate a virtual environment?"
+        ) from exc
+    execute_from_command_line(sys.argv)
+
+
+if __name__ == '__main__':
+    main()
diff --git a/Engine/Source/Programs/WebTestsServerPy/requirements.txt b/Engine/Source/Programs/WebTestsServerPy/requirements.txt
new file mode 100644
--- /dev/null
+++ b/Engine/Source/Programs/WebTestsServerPy/requirements.txt
@@ -0,0 +1,5 @@
+autobahn==24.4.2
+django==4.2.0
+djangorestframework==3.14.0
+channels==4.0.0
+daphne==4.0.0
diff --git a/Engine/Source/Programs/WebTestsServerPy/runserver.bat b/Engine/Source/Programs/WebTestsServerPy/runserver.bat
new file mode 100755
--- /dev/null
+++ b/Engine/Source/Programs/WebTestsServerPy/runserver.bat
@@ -0,0 +1,5 @@
+@echo off
+
+call createenv.bat
+
+python -m daphne -e tcp:port=8000:interface=0.0.0.0 -e tcp:port=8001:interface=0.0.0.0 webtests.asgi:application
diff --git a/Engine/Source/Programs/WebTestsServerPy/runserver.sh b/Engine/Source/Programs/WebTestsServerPy/runserver.sh
new file mode 100644
--- /dev/null
+++ b/Engine/Source/Programs/WebTestsServerPy/runserver.sh
@@ -0,0 +1,2 @@
+sh createenv.sh
+python -m daphne -e tcp:port=8000:interface=0.0.0.0 -e tcp:port=8001:interface=0.0.0.0 webtests.asgi:application
diff --git a/Engine/Source/Programs/WebTestsServerPy/runserver_unixsocket.bat b/Engine/Source/Programs/WebTestsServerPy/runserver_unixsocket.bat
new file mode 100755
--- /dev/null
+++ b/Engine/Source/Programs/WebTestsServerPy/runserver_unixsocket.bat
@@ -0,0 +1,2 @@
+dotnet build .\unixsockets\unixsocketsserver.csproj -c "Release"
+dotnet run --project ".\unixsockets\unixsocketsserver.csproj" -c "Release"
\ No newline at end of file
diff --git a/Engine/Source/Programs/WebTestsServerPy/unixsockets/Program.cs b/Engine/Source/Programs/WebTestsServerPy/unixsockets/Program.cs
new file mode 100644
--- /dev/null
+++ b/Engine/Source/Programs/WebTestsServerPy/unixsockets/Program.cs
@@ -0,0 +1,48 @@
+﻿// Copyright Epic Games, Inc. All Rights Reserved.
+using Microsoft.AspNetCore.Builder;
+using Microsoft.AspNetCore.Hosting;
+
+const string SocketFileName = "webtests.sock";
+
+string CurrentDirectory = Environment.CurrentDirectory;
+string SocketsDirectory = Path.Combine(CurrentDirectory, "sockets");
+
+if (Directory.Exists(SocketsDirectory))
+{
+    //Clean any sockets remaining from a previous run
+    foreach(var SocketFile in Directory.EnumerateFiles(SocketsDirectory, "*", SearchOption.AllDirectories))
+    {
+        File.Delete(SocketFile);
+    }
+}
+else
+{
+    Directory.CreateDirectory(SocketsDirectory);
+}
+
+string SocketPath = Path.Combine(SocketsDirectory, SocketFileName);
+Console.WriteLine($"Run Web Tests with the following arguments to test Unix Sockets");
+Console.WriteLine($"=> \"Http Methods over Unix Domain Socket\" --extra-args -very_verbose=true -web_server_unix_socket=\"{SocketPath}\"");
+
+var AppBuilder = WebApplication.CreateBuilder();
+AppBuilder.WebHost.UseKestrel(Options =>
+{
+    Options.ListenUnixSocket(SocketPath);
+});
+
+var WebServerApp = AppBuilder.Build();
+
+string[] Methods = { "GET", "POST", "PUT", "DELETE" };
+
+WebServerApp.MapMethods("webtests/unixsockettests/{param}", Methods, async Context =>
+{
+    string Param = (string)Context.Request.RouteValues["param"];
+    Context.Response.StatusCode = 200;
+    await using (StreamWriter Writer = new StreamWriter(Context.Response.Body, leaveOpen:true))
+    {
+        await Writer.WriteAsync(Param);
+    }
+    await Task.CompletedTask;
+});
+
+WebServerApp.Run();
\ No newline at end of file
diff --git a/Engine/Source/Programs/WebTestsServerPy/unixsockets/unixsocketsserver.csproj b/Engine/Source/Programs/WebTestsServerPy/unixsockets/unixsocketsserver.csproj
new file mode 100644
--- /dev/null
+++ b/Engine/Source/Programs/WebTestsServerPy/unixsockets/unixsocketsserver.csproj
@@ -0,0 +1,14 @@
+﻿<Project Sdk="Microsoft.NET.Sdk">
+
+  <PropertyGroup>
+    <OutputType>Exe</OutputType>
+    <TargetFramework>net8.0</TargetFramework>
+    <ImplicitUsings>enable</ImplicitUsings>
+    <Nullable>disable</Nullable>
+  </PropertyGroup>
+
+  <ItemGroup>
+    <FrameworkReference Include="Microsoft.AspNetCore.App" />
+  </ItemGroup>
+
+</Project>
diff --git a/Engine/Source/Programs/WebTestsServerPy/websocketstests/consumers.py b/Engine/Source/Programs/WebTestsServerPy/websocketstests/consumers.py
new file mode 100644
--- /dev/null
+++ b/Engine/Source/Programs/WebTestsServerPy/websocketstests/consumers.py
@@ -0,0 +1,23 @@
+import json
+
+from channels.generic.websocket import WebsocketConsumer
+
+class EchoConsumer(WebsocketConsumer):
+    def connect(self):
+        self.accept()
+
+    def disconnect(self, close_code):
+        pass
+
+    def receive(self, text_data):
+        self.send(text_data)
+
+class CloseOnReceiveMessageConsumer(WebsocketConsumer):
+    def connect(self):
+        self.accept()
+
+    def disconnect(self, close_code):
+        pass
+
+    def receive(self, text_data):
+        self.close()
diff --git a/Engine/Source/Programs/WebTestsServerPy/websocketstests/routing.py b/Engine/Source/Programs/WebTestsServerPy/websocketstests/routing.py
new file mode 100644
--- /dev/null
+++ b/Engine/Source/Programs/WebTestsServerPy/websocketstests/routing.py
@@ -0,0 +1,8 @@
+from django.urls import re_path
+
+from . import consumers
+
+websocket_urlpatterns = [
+    re_path(r"webtests/websocketstests/echo", consumers.EchoConsumer.as_asgi()),
+    re_path(r"webtests/websocketstests/close_on_receive_message", consumers.CloseOnReceiveMessageConsumer.as_asgi()),
+]
diff --git a/Engine/Source/Programs/WebTestsServerPy/websocketstests/templates/websocketstests/index.html b/Engine/Source/Programs/WebTestsServerPy/websocketstests/templates/websocketstests/index.html
new file mode 100644
--- /dev/null
+++ b/Engine/Source/Programs/WebTestsServerPy/websocketstests/templates/websocketstests/index.html
@@ -0,0 +1,44 @@
+<!DOCTYPE html>
+<html>
+<head>
+    <meta charset="utf-8"/>
+    <title>WebSockects Tests</title>
+</head>
+<body>
+    <textarea id="echo-log" cols="100" rows="20"></textarea><br>
+    <input id="echo-message-input" type="text" size="100"><br>
+    <input id="echo-message-submit" type="button" value="Send">
+    <script>
+        const echoSocket = new WebSocket(
+            'ws://'
+            + window.location.host
+            + '/webtests/websocketstests/echo/'
+        );
+
+        echoSocket.onmessage = function(e) {
+            const data = JSON.parse(e.data);
+            document.querySelector('#echo-log').value += (data.message + '\n');
+        };
+
+        echoSocket.onclose = function(e) {
+            console.error('Echo socket closed unexpectedly');
+        };
+
+        document.querySelector('#echo-message-input').focus();
+        document.querySelector('#echo-message-input').onkeyup = function(e) {
+            if (e.key === 'Enter') {  // enter, return
+                document.querySelector('#echo-message-submit').click();
+            }
+        };
+
+        document.querySelector('#echo-message-submit').onclick = function(e) {
+            const messageInputDom = document.querySelector('#echo-message-input');
+            const message = messageInputDom.value;
+            echoSocket.send(JSON.stringify({
+                'message': message
+            }));
+            messageInputDom.value = '';
+        };
+    </script>
+</body>
+</html>
diff --git a/Engine/Source/Programs/WebTestsServerPy/websocketstests/urls.py b/Engine/Source/Programs/WebTestsServerPy/websocketstests/urls.py
new file mode 100644
--- /dev/null
+++ b/Engine/Source/Programs/WebTestsServerPy/websocketstests/urls.py
@@ -0,0 +1,7 @@
+from django.urls import path
+
+from . import views
+
+urlpatterns = [
+    path("websocketstests", views.index, name="index"),
+]
diff --git a/Engine/Source/Programs/WebTestsServerPy/websocketstests/views.py b/Engine/Source/Programs/WebTestsServerPy/websocketstests/views.py
new file mode 100644
--- /dev/null
+++ b/Engine/Source/Programs/WebTestsServerPy/websocketstests/views.py
@@ -0,0 +1,5 @@
+from django.shortcuts import render
+
+def index(request):
+    return render(request, "websocketstests/index.html")
+
diff --git a/Engine/Source/Programs/WebTestsServerPy/webtests/asgi.py b/Engine/Source/Programs/WebTestsServerPy/webtests/asgi.py
new file mode 100644
--- /dev/null
+++ b/Engine/Source/Programs/WebTestsServerPy/webtests/asgi.py
@@ -0,0 +1,28 @@
+"""
+ASGI config for webtests project.
+
+It exposes the ASGI callable as a module-level variable named ``application``.
+
+For more information on this file, see
+https://docs.djangoproject.com/en/4.1/howto/deployment/asgi/
+"""
+
+import os
+
+from channels.auth import AuthMiddlewareStack
+from channels.routing import ProtocolTypeRouter, URLRouter
+from channels.security.websocket import AllowedHostsOriginValidator
+from django.core.asgi import get_asgi_application
+
+from websocketstests.routing import websocket_urlpatterns
+
+os.environ.setdefault('DJANGO_SETTINGS_MODULE', 'webtests.settings')
+
+application = ProtocolTypeRouter(
+    {
+        "http": get_asgi_application(),
+        "websocket": AllowedHostsOriginValidator(
+            AuthMiddlewareStack(URLRouter(websocket_urlpatterns))
+        ),
+    }
+)
diff --git a/Engine/Source/Programs/WebTestsServerPy/webtests/settings.py b/Engine/Source/Programs/WebTestsServerPy/webtests/settings.py
new file mode 100644
--- /dev/null
+++ b/Engine/Source/Programs/WebTestsServerPy/webtests/settings.py
@@ -0,0 +1,135 @@
+"""
+Django settings for webtests project.
+
+Generated by 'django-admin startproject' using Django 4.1.7.
+
+For more information on this file, see
+https://docs.djangoproject.com/en/4.1/topics/settings/
+
+For the full list of settings and their values, see
+https://docs.djangoproject.com/en/4.1/ref/settings/
+"""
+
+from pathlib import Path
+
+# Build paths inside the project like this: BASE_DIR / 'subdir'.
+BASE_DIR = Path(__file__).resolve().parent.parent
+
+
+# Quick-start development settings - unsuitable for production
+# See https://docs.djangoproject.com/en/4.1/howto/deployment/checklist/
+
+# SECURITY WARNING: keep the secret key used in production secret!
+SECRET_KEY = 'django-insecure-x=3k^30px$8#ouewp0(1er-gpm@90k^z5$6bxx(qtot92yy$%w'
+
+# SECURITY WARNING: don't run with debug turned on in production!
+DEBUG = True
+
+# Allow any host, because this web server is only for local test purpose
+ALLOWED_HOSTS = ['*']
+
+
+# Application definition
+
+INSTALLED_APPS = [
+    'daphne',
+    'websocketstests',
+    'django.contrib.admin',
+    'django.contrib.auth',
+    'django.contrib.contenttypes',
+    'django.contrib.sessions',
+    'django.contrib.messages',
+    'django.contrib.staticfiles',
+]
+
+MIDDLEWARE = [
+    'django.middleware.security.SecurityMiddleware',
+    'django.contrib.sessions.middleware.SessionMiddleware',
+    'django.middleware.common.CommonMiddleware',
+    'django.middleware.csrf.CsrfViewMiddleware',
+    'django.contrib.auth.middleware.AuthenticationMiddleware',
+    'django.contrib.messages.middleware.MessageMiddleware',
+    'django.middleware.clickjacking.XFrameOptionsMiddleware',
+]
+
+ROOT_URLCONF = 'webtests.urls'
+
+TEMPLATES = [
+    {
+        'BACKEND': 'django.template.backends.django.DjangoTemplates',
+        'DIRS': [],
+        'APP_DIRS': True,
+        'OPTIONS': {
+            'context_processors': [
+                'django.template.context_processors.debug',
+                'django.template.context_processors.request',
+                'django.contrib.auth.context_processors.auth',
+                'django.contrib.messages.context_processors.messages',
+            ],
+        },
+    },
+]
+
+WSGI_APPLICATION = 'webtests.wsgi.application'
+
+
+# Database
+# https://docs.djangoproject.com/en/4.1/ref/settings/#databases
+
+DATABASES = {
+    # No need of database for web tests, at least for now
+    # 'default': {
+    #     'ENGINE': 'django.db.backends.sqlite3',
+    #     'NAME': BASE_DIR / 'db.sqlite3',
+    # }
+}
+
+
+# Password validation
+# https://docs.djangoproject.com/en/4.1/ref/settings/#auth-password-validators
+
+AUTH_PASSWORD_VALIDATORS = [
+    {
+        'NAME': 'django.contrib.auth.password_validation.UserAttributeSimilarityValidator',
+    },
+    {
+        'NAME': 'django.contrib.auth.password_validation.MinimumLengthValidator',
+    },
+    {
+        'NAME': 'django.contrib.auth.password_validation.CommonPasswordValidator',
+    },
+    {
+        'NAME': 'django.contrib.auth.password_validation.NumericPasswordValidator',
+    },
+]
+
+
+# Internationalization
+# https://docs.djangoproject.com/en/4.1/topics/i18n/
+
+LANGUAGE_CODE = 'en-us'
+
+TIME_ZONE = 'UTC'
+
+USE_I18N = True
+
+USE_TZ = True
+
+
+# Static files (CSS, JavaScript, Images)
+# https://docs.djangoproject.com/en/4.1/howto/static-files/
+
+STATIC_URL = 'static/'
+
+# Default primary key field type
+# https://docs.djangoproject.com/en/4.1/ref/settings/#default-auto-field
+
+DEFAULT_AUTO_FIELD = 'django.db.models.BigAutoField'
+
+LOGGING = {
+    "version": 1,  # the dictConfig format version
+    "disable_existing_loggers": False,  # retain the default loggers
+    "level": "DEBUG",
+}
+
+ASGI_APPLICATION = "webtests.asgi.application"
diff --git a/Engine/Source/Programs/WebTestsServerPy/webtests/urls.py b/Engine/Source/Programs/WebTestsServerPy/webtests/urls.py
new file mode 100644
--- /dev/null
+++ b/Engine/Source/Programs/WebTestsServerPy/webtests/urls.py
@@ -0,0 +1,23 @@
+"""webtests URL Configuration
+
+The `urlpatterns` list routes URLs to views. For more information please see:
+    https://docs.djangoproject.com/en/4.1/topics/http/urls/
+Examples:
+Function views
+    1. Add an import:  from my_app import views
+    2. Add a URL to urlpatterns:  path('', views.home, name='home')
+Class-based views
+    1. Add an import:  from other_app.views import Home
+    2. Add a URL to urlpatterns:  path('', Home.as_view(), name='home')
+Including another URLconf
+    1. Import the include() function: from django.urls import include, path
+    2. Add a URL to urlpatterns:  path('blog/', include('blog.urls'))
+"""
+from django.contrib import admin
+from django.urls import include, path
+
+urlpatterns = [
+    path("webtests/", include("httptests.urls")),
+    path("webtests/", include("websocketstests.urls")),
+    path('admin/', admin.site.urls),
+]
diff --git a/Engine/Source/Programs/WebTestsServerPy/webtests/wsgi.py b/Engine/Source/Programs/WebTestsServerPy/webtests/wsgi.py
new file mode 100644
--- /dev/null
+++ b/Engine/Source/Programs/WebTestsServerPy/webtests/wsgi.py
@@ -0,0 +1,16 @@
+"""
+WSGI config for webtests project.
+
+It exposes the WSGI callable as a module-level variable named ``application``.
+
+For more information on this file, see
+https://docs.djangoproject.com/en/4.1/howto/deployment/wsgi/
+"""
+
+import os
+
+from django.core.wsgi import get_wsgi_application
+
+os.environ.setdefault('DJANGO_SETTINGS_MODULE', 'webtests.settings')
+
+application = get_wsgi_application()
